Add basic UObject thread safety checks to UnrealCodeAnalyzer.

diff --git a/Engine/Source/Developer/UnrealCodeAnalyzerTests/Private/ThreadSafety.h b/Engine/Source/Developer/UnrealCodeAnalyzerTests/Private/ThreadSafety.h
new file mode 100644
--- /dev/null
+++ b/Engine/Source/Developer/UnrealCodeAnalyzerTests/Private/ThreadSafety.h
@@ -0,0 +1,16 @@
+// Copyright 1998-2015 Epic Games, Inc. All Rights Reserved.
+
+#pragma once
+#include "ThreadSafety.generated.h"
+
+UCLASS()
+class UObjectDerivedThreadSafetyTest : public UObject
+{
+	GENERATED_BODY()
+
+public:
+	UObjectDerivedThreadSafetyTest();
+
+	virtual void PostInitProperties() override;
+	virtual void Serialize(FArchive& Ar) override;
+};
\ No newline at end of file
diff --git a/Engine/Source/Developer/UnrealCodeAnalyzerTests/Public/UnrealCodeAnalyzerTests.h b/Engine/Source/Developer/UnrealCodeAnalyzerTests/Public/UnrealCodeAnalyzerTests.h
--- a/Engine/Source/Developer/UnrealCodeAnalyzerTests/Public/UnrealCodeAnalyzerTests.h
+++ b/Engine/Source/Developer/UnrealCodeAnalyzerTests/Public/UnrealCodeAnalyzerTests.h
@@ -2,4 +2,4 @@
 
 #pragma once
 
-#include "Core.h"
+#include "CoreUObject.h"
diff --git a/Engine/Source/Developer/UnrealCodeAnalyzerTests/UnrealCodeAnalyzerTests.Build.cs b/Engine/Source/Developer/UnrealCodeAnalyzerTests/UnrealCodeAnalyzerTests.Build.cs
--- a/Engine/Source/Developer/UnrealCodeAnalyzerTests/UnrealCodeAnalyzerTests.Build.cs
+++ b/Engine/Source/Developer/UnrealCodeAnalyzerTests/UnrealCodeAnalyzerTests.Build.cs
@@ -7,5 +7,6 @@ public class UnrealCodeAnalyzerTests : ModuleRules
 	public UnrealCodeAnalyzerTests(TargetInfo Target)
 	{
 		PublicDependencyModuleNames.Add("Core");
+		PublicDependencyModuleNames.Add("CoreUObject");
 	}
 }
diff --git a/Engine/Source/Programs/UnrealBuildTool/Configuration/BuildConfiguration.cs b/Engine/Source/Programs/UnrealBuildTool/Configuration/BuildConfiguration.cs
--- a/Engine/Source/Programs/UnrealBuildTool/Configuration/BuildConfiguration.cs
+++ b/Engine/Source/Programs/UnrealBuildTool/Configuration/BuildConfiguration.cs
@@ -426,11 +426,23 @@ namespace UnrealBuildTool
 		public static string DistccExecutablesPath;
 
 		/// <summary>
-		/// Run UnrealCodeAnalyzer instead of regular compilation.
+		/// Run UnrealCodeAnalyzer instead of regular compilation. Requires setting UCA mode by bUCACheckUObjectThreadSafety or bUCACheckPCHFiles.
 		/// </summary>
 		[XmlConfig]
 		public static bool bRunUnrealCodeAnalyzer;
 
+		/// <summary>
+		/// Perform basic UObject thread safety checks. Dumps information whether UObject constructors, Serialize and PostInitProperties functions access any global/static variables or functions.
+		/// </summary>
+		[XmlConfig]
+		public static bool bUCACheckUObjectThreadSafety;
+
+		/// <summary>
+		/// Dumps information about header inclusion in PCH.
+		/// </summary>
+		[XmlConfig]
+		public static bool bUCACheckPCHFiles;
+
 		/// <summary>
 		/// Percentage of cpp files in module where header must be included to get included in PCH.
 		/// </summary>
@@ -570,6 +582,8 @@ namespace UnrealBuildTool
 			// If header is included in 0% or more cpp files in module it should be included in PCH.
 			UCAUsageThreshold = 100.0f;
 
+			bUCACheckUObjectThreadSafety = false;
+
 			// The default for normal Mac users should be to use DistCode which installs as an Xcode plugin and provides dynamic host management
 			if (BuildHostPlatform.Current.Platform == UnrealTargetPlatform.Mac)
 			{
@@ -672,8 +686,8 @@ namespace UnrealBuildTool
 				BuildConfiguration.bUsePCHFiles = false;
 				BuildConfiguration.bUseSharedPCHs = false;
 
-				//BuildConfiguration.bUseUBTMakefiles = false;
-				//BuildConfiguration.bUseActionHistory = false;
+				BuildConfiguration.bUseUBTMakefiles = false;
+				BuildConfiguration.bUseActionHistory = false;
 			}
 		}
 	}
diff --git a/Engine/Source/Programs/UnrealBuildTool/System/UnrealBuildTool.cs b/Engine/Source/Programs/UnrealBuildTool/System/UnrealBuildTool.cs
--- a/Engine/Source/Programs/UnrealBuildTool/System/UnrealBuildTool.cs
+++ b/Engine/Source/Programs/UnrealBuildTool/System/UnrealBuildTool.cs
@@ -836,6 +836,14 @@ namespace UnrealBuildTool
                 {
                     BuildConfiguration.UCAModuleToAnalyze = LowercaseArg.Substring(UCAModuleToAnalyzeString.Length).Trim();
                 }
+				else if (LowercaseArg == "-ucacheckuobjectthreadsafety")
+				{
+					BuildConfiguration.bUCACheckUObjectThreadSafety = true;
+				}
+				else if (LowercaseArg == "-ucacheckpchfiles")
+				{
+					BuildConfiguration.bUCACheckPCHFiles = true;
+				}
             }
         }
 
@@ -1150,11 +1158,20 @@ namespace UnrealBuildTool
                         throw new BuildException( "UnrealBuildTool: At least one of either IsGatheringBuild or IsAssemblingBuild must be true.  Did you pass '-NoGather' with '-NoAssemble'?" );
                     }
 
-					if (BuildConfiguration.bRunUnrealCodeAnalyzer && (BuildConfiguration.UCAModuleToAnalyze == null || BuildConfiguration.UCAModuleToAnalyze.Length == 0))
+					if (BuildConfiguration.bRunUnrealCodeAnalyzer)
+					{
+						if (BuildConfiguration.UCAModuleToAnalyze == null || BuildConfiguration.UCAModuleToAnalyze.Length == 0)
 					{
 						throw new BuildException("When running UnrealCodeAnalyzer, please specify module to analyze in UCAModuleToAnalyze field in BuildConfiguration.xml");
 					}
 
+						if ((BuildConfiguration.bUCACheckPCHFiles && BuildConfiguration.bUCACheckUObjectThreadSafety)
+							|| (!BuildConfiguration.bUCACheckPCHFiles && !BuildConfiguration.bUCACheckUObjectThreadSafety))
+						{
+							throw new BuildException("When running UnrealCodeAnalyzer, please specify exactly one action to execute");
+						}
+					}
+
 					// Get the build version
 					BuildVersion Version;
 					if(!Utils.TryReadBuildVersion("../Build/Build.version", out Version))
diff --git a/Engine/Source/Programs/UnrealBuildTool/Windows/VCToolChain.cs b/Engine/Source/Programs/UnrealBuildTool/Windows/VCToolChain.cs
--- a/Engine/Source/Programs/UnrealBuildTool/Windows/VCToolChain.cs
+++ b/Engine/Source/Programs/UnrealBuildTool/Windows/VCToolChain.cs
@@ -1309,10 +1309,11 @@ namespace UnrealBuildTool
 				string UnrealCodeAnalyzerArguments = "";
 				if (BuildConfiguration.bRunUnrealCodeAnalyzer)
 				{
-					var ObjectFileExtension = @".includes";
+					string UCAMode = BuildConfiguration.bUCACheckUObjectThreadSafety ? @"-CheckThreadSafety " : @"-CreateIncludeFiles ";
+					var ObjectFileExtension = BuildConfiguration.bUCACheckUObjectThreadSafety ? @".tsc" : @".includes";
 					FileItem ObjectFile = FileItem.GetItemByPath(Path.Combine(CompileEnvironment.Config.OutputDirectory, Path.GetFileName(SourceFile.AbsolutePath) + ObjectFileExtension));
 					var ClangPath = System.IO.Path.Combine(CompileAction.WorkingDirectory, @"ThirdParty", @"llvm", @"3.5.0", @"bin", @"vs2013", @"x86", @"release", @"clang++.exe");
-					UnrealCodeAnalyzerArguments = @"-CreateIncludeFiles " + SourceFile.AbsolutePath + @" -OutputFile=""" + ObjectFile.AbsolutePath + @""" -- " + ClangPath + @" --driver-mode=cl ";
+					UnrealCodeAnalyzerArguments = UCAMode + SourceFile.AbsolutePath + @" -OutputFile=""" + ObjectFile.AbsolutePath + @""" -- " + ClangPath + @" --driver-mode=cl ";
 				}
 
 				CompileAction.CommandArguments = UnrealCodeAnalyzerArguments + SharedArguments.ToString() + FileArguments.ToString() + CompileEnvironment.Config.AdditionalArguments;
diff --git a/Engine/Source/Programs/UnrealCodeAnalyzer/Private/ThreadSafety/TSAction.cpp b/Engine/Source/Programs/UnrealCodeAnalyzer/Private/ThreadSafety/TSAction.cpp
new file mode 100644
--- /dev/null
+++ b/Engine/Source/Programs/UnrealCodeAnalyzer/Private/ThreadSafety/TSAction.cpp
@@ -0,0 +1,27 @@
+// Copyright 1998-2015 Epic Games, Inc. All Rights Reserved.
+
+#include "UnrealCodeAnalyzerPCH.h"
+#include "ThreadSafety/TSAction.h"
+#include "ThreadSafety/TSConsumer.h"
+
+namespace UnrealCodeAnalyzer
+{
+
+	clang::ASTConsumer* FTSAction::CreateASTConsumer(clang::CompilerInstance& Compiler, clang::StringRef InFile)
+	{
+		Filename = InFile;
+		SourceManager = &Compiler.getSourceManager();
+		Consumer = new FTSConsumer(&Compiler.getASTContext(), InFile);
+		return Consumer;
+	}
+
+	bool FTSAction::BeginSourceFileAction(clang::CompilerInstance &CI, clang::StringRef Filename)
+	{
+		return true;
+	}
+
+	void FTSAction::EndSourceFileAction()
+	{
+	}
+
+}
diff --git a/Engine/Source/Programs/UnrealCodeAnalyzer/Private/ThreadSafety/TSAction.h b/Engine/Source/Programs/UnrealCodeAnalyzer/Private/ThreadSafety/TSAction.h
new file mode 100644
--- /dev/null
+++ b/Engine/Source/Programs/UnrealCodeAnalyzer/Private/ThreadSafety/TSAction.h
@@ -0,0 +1,23 @@
+// Copyright 1998-2015 Epic Games, Inc. All Rights Reserved.
+
+#pragma once
+#include "clang/Frontend/FrontendAction.h"
+
+namespace UnrealCodeAnalyzer
+{
+	class FTSConsumer;
+	class clang::SourceManager;
+
+	class FTSAction : public clang::ASTFrontendAction
+	{
+	public:
+		virtual clang::ASTConsumer* CreateASTConsumer(clang::CompilerInstance& Compiler, clang::StringRef InFile) override;
+		virtual bool BeginSourceFileAction(clang::CompilerInstance &CI, clang::StringRef Filename) override;
+		virtual void EndSourceFileAction() override;
+
+	private:
+		FTSConsumer* Consumer;
+		clang::StringRef Filename;
+		clang::SourceManager* SourceManager;
+	};
+}
diff --git a/Engine/Source/Programs/UnrealCodeAnalyzer/Private/ThreadSafety/TSConsumer.cpp b/Engine/Source/Programs/UnrealCodeAnalyzer/Private/ThreadSafety/TSConsumer.cpp
new file mode 100644
--- /dev/null
+++ b/Engine/Source/Programs/UnrealCodeAnalyzer/Private/ThreadSafety/TSConsumer.cpp
@@ -0,0 +1,19 @@
+// Copyright 1998-2015 Epic Games, Inc. All Rights Reserved.
+
+#include "UnrealCodeAnalyzerPCH.h"
+#include "ThreadSafety/TSConsumer.h"
+
+using namespace clang;
+
+namespace UnrealCodeAnalyzer
+{
+	void FTSConsumer::HandleTranslationUnit(ASTContext& Context)
+	{
+		Visitor.TraverseDecl(Context.getTranslationUnitDecl());
+	}
+
+	FTSConsumer::FTSConsumer(ASTContext* Context, StringRef InFile) :
+		Visitor(Context, InFile)
+	{ }
+
+}
diff --git a/Engine/Source/Programs/UnrealCodeAnalyzer/Private/ThreadSafety/TSConsumer.h b/Engine/Source/Programs/UnrealCodeAnalyzer/Private/ThreadSafety/TSConsumer.h
new file mode 100644
--- /dev/null
+++ b/Engine/Source/Programs/UnrealCodeAnalyzer/Private/ThreadSafety/TSConsumer.h
@@ -0,0 +1,19 @@
+// Copyright 1998-2015 Epic Games, Inc. All Rights Reserved.
+
+#pragma once
+#include "clang/AST/ASTConsumer.h"
+#include "TSVisitor.h"
+
+namespace UnrealCodeAnalyzer
+{
+	class FTSConsumer : public clang::ASTConsumer
+	{
+	public:
+		explicit FTSConsumer(clang::ASTContext* Context, clang::StringRef InFile);
+
+		virtual void HandleTranslationUnit(clang::ASTContext &Ctx) override;
+
+	private:
+		FTSVisitor Visitor;
+	};
+}
diff --git a/Engine/Source/Programs/UnrealCodeAnalyzer/Private/ThreadSafety/TSVisitor.cpp b/Engine/Source/Programs/UnrealCodeAnalyzer/Private/ThreadSafety/TSVisitor.cpp
new file mode 100644
--- /dev/null
+++ b/Engine/Source/Programs/UnrealCodeAnalyzer/Private/ThreadSafety/TSVisitor.cpp
@@ -0,0 +1,189 @@
+// Copyright 1998-2015 Epic Games, Inc. All Rights Reserved.
+
+#include "UnrealCodeAnalyzerPCH.h"
+#include "Visitor.h"
+
+namespace UnrealCodeAnalyzer
+{
+	FTSVisitor::FTSVisitor(clang::ASTContext* InContext, clang::StringRef InFile)
+		: SourceManager(InContext->getSourceManager())
+	{ }
+
+	bool FTSVisitor::VisitCXXConstructorDecl(clang::CXXConstructorDecl* Decl)
+	{
+		if (!Decl->isThisDeclarationADefinition())
+		{
+			return true;
+		}
+
+		if (!InheritsFromUObject(Decl))
+		{
+			return true;
+		}
+
+		bool bVisited;
+		VisitedDecls.Add(Decl, &bVisited);
+
+		if (!bVisited)
+		{
+			FunctionDeclarationsStack.Add(Decl);
+			RecursiveASTVisitor<FTSVisitor>::TraverseCXXConstructorDecl(Decl);
+			FunctionDeclarationsStack.Pop(false);
+		}
+
+		return true;
+	}
+
+	bool FTSVisitor::VisitDeclRefExpr(clang::DeclRefExpr* Expr)
+	{
+		// Get declaration for the expression we found.
+		clang::NamedDecl* NamedDecl = Expr->getFoundDecl();
+
+		// If the expression is not a variable, we have nothing to do here.
+		clang::VarDecl* VarDecl = llvm::dyn_cast<clang::VarDecl>(NamedDecl);
+
+		if (VarDecl == nullptr)
+		{
+			return true;
+		}
+
+		if (IsGlobalVariable(VarDecl) && FunctionDeclarationsStack.Num() > 0)
+		{
+			clang::SourceLocation Begin = Expr->getExprLoc();
+			std::pair<clang::FileID, unsigned> DecomposedLoc = SourceManager.getDecomposedLoc(Begin);
+			int32 Line = SourceManager.getLineNumber(DecomposedLoc.first, DecomposedLoc.second);
+			int32 Column = SourceManager.getColumnNumber(DecomposedLoc.first, DecomposedLoc.second);
+			clang::StringRef Filename = SourceManager.getFilename(Begin);
+			const char* Data = Filename.data();
+			if (Line == 1 && Column == 1 && Data == '\0')
+			{
+				clang::SourceLocation Begin = FunctionDeclarationsStack.Last()->getLocation();
+				std::pair<clang::FileID, unsigned> DecomposedLoc = SourceManager.getDecomposedLoc(Begin);
+				int32 Line = SourceManager.getLineNumber(DecomposedLoc.first, DecomposedLoc.second);
+				int32 Column = SourceManager.getColumnNumber(DecomposedLoc.first, DecomposedLoc.second);
+				clang::StringRef Filename = SourceManager.getFilename(Begin);
+				const char* Data = Filename.data();
+				OutputFileContents.Logf(TEXT("[FuncLoc]%s:%d:%d Found usage of %s in function %s::%s"), *FString(Data), Line, Column, *FString(VarDecl->getName().data()), ANSI_TO_TCHAR(FunctionDeclarationsStack.Last()->getParent()->getNameAsString().c_str()), *FString(FunctionDeclarationsStack.Last()->getNameAsString().c_str()));
+			}
+			else
+			{
+				OutputFileContents.Logf(TEXT("[ExprLoc]%s:%d:%d Found usage of %s in function %s::%s"), *FString(Data), Line, Column, *FString(VarDecl->getName().data()), ANSI_TO_TCHAR(FunctionDeclarationsStack.Last()->getParent()->getNameAsString().c_str()), *FString(FunctionDeclarationsStack.Last()->getNameAsString().c_str()));
+			}
+		}
+		return true;
+	}
+
+	bool FTSVisitor::TraverseCXXMethodDecl(clang::CXXMethodDecl* Decl)
+	{
+		if (!Decl->isThisDeclarationADefinition())
+		{
+			return true;
+		}
+		if ((strcmp(Decl->getName().data(), "PostInitProperties") == 0)
+			|| (strcmp(Decl->getName().data(), "Serialize") == 0))
+		{
+			FunctionDeclarationsStack.Add(Decl);
+			RecursiveASTVisitor<FTSVisitor>::TraverseFunctionDecl(Decl);
+			FunctionDeclarationsStack.Pop(false);
+		}
+
+		return true;
+	}
+
+	bool FTSVisitor::InheritsFromUObject(clang::CXXConstructorDecl* Decl)
+	{
+		const clang::CXXRecordDecl* ParentClass = Decl->getParent();
+		return ParentClass->getNameAsString()[0] == 'U';
+	}
+
+	bool FTSVisitor::IsGlobalVariable(clang::VarDecl* VarDecl)
+	{
+		// Get context of declaration.
+		clang::DeclContext* DeclContext = VarDecl->getDeclContext();
+
+		// Globals are declared in the scope of translation unit
+		if (DeclContext->isTranslationUnit()
+			// Or are static data members
+			|| VarDecl->isStaticDataMember())
+		{
+			return true;
+		}
+
+		if (DeclContext->isNamespace())
+		{
+			// Go to outer contexts until there is none or outer is translation unit.
+			clang::DeclContext* OuterContext = DeclContext->getLookupParent();
+
+			while (OuterContext != nullptr)
+			{
+				if (OuterContext->isTranslationUnit())
+				{
+					return true;
+				}
+
+				OuterContext = OuterContext->getLookupParent();
+			}
+		}
+
+		return false;
+	}
+
+	bool FTSVisitor::VisitCallExpr(clang::CallExpr* Expr)
+	{
+		bool bIsGlobal = false;
+		clang::FunctionDecl* Decl = dyn_cast_or_null<clang::FunctionDecl>(Expr->getCalleeDecl());
+		if (Decl)
+		{
+			if (Decl->isGlobal())
+			{
+				bIsGlobal = true;
+			}
+			else if (clang::CXXMethodDecl* MethodDecl = dyn_cast_or_null<clang::CXXMethodDecl>(Decl))
+			{
+				if (MethodDecl->isStatic())
+				{
+					bIsGlobal = true;
+				}
+			}
+		}
+
+		if (!bIsGlobal || FunctionDeclarationsStack.Num() == 0)
+		{
+			return true;
+		}
+
+		FString FunctionName;
+		if (clang::CXXMethodDecl* MethodDecl = dyn_cast_or_null<clang::CXXMethodDecl>(Decl))
+		{
+			FunctionName = FString::Printf(TEXT("%s::%s"), ANSI_TO_TCHAR(MethodDecl->getParent()->getNameAsString().c_str()), ANSI_TO_TCHAR(Decl->getNameAsString().c_str()));
+		}
+		else
+		{
+			FunctionName = FString(ANSI_TO_TCHAR(Decl->getNameAsString().c_str()));
+		}
+
+		clang::SourceLocation Begin = Expr->getExprLoc();
+		std::pair<clang::FileID, unsigned> DecomposedLoc = SourceManager.getDecomposedLoc(Begin);
+		int32 Line = SourceManager.getLineNumber(DecomposedLoc.first, DecomposedLoc.second);
+		int32 Column = SourceManager.getColumnNumber(DecomposedLoc.first, DecomposedLoc.second);
+		clang::StringRef Filename = SourceManager.getFilename(Begin);
+		const char* Data = Filename.data();
+		if (Line == 1 && Column == 1 && Data == '\0')
+		{
+			clang::SourceLocation Begin = FunctionDeclarationsStack.Last()->getLocation();
+			std::pair<clang::FileID, unsigned> DecomposedLoc = SourceManager.getDecomposedLoc(Begin);
+			int32 Line = SourceManager.getLineNumber(DecomposedLoc.first, DecomposedLoc.second);
+			int32 Column = SourceManager.getColumnNumber(DecomposedLoc.first, DecomposedLoc.second);
+			clang::StringRef Filename = SourceManager.getFilename(Begin);
+			const char* Data = Filename.data();
+			OutputFileContents.Logf(TEXT("[FuncLoc]%s:%d:%d Found usage of %s in function %s::%s"), *FString(Data), Line, Column, *FunctionName, ANSI_TO_TCHAR(FunctionDeclarationsStack.Last()->getParent()->getNameAsString().c_str()), *FString(FunctionDeclarationsStack.Last()->getNameAsString().c_str()));
+		}
+		else
+		{
+			OutputFileContents.Logf(TEXT("[ExprLoc]%s:%d:%d Found usage of %s in function %s::%s"), *FString(Data), Line, Column, *FunctionName, ANSI_TO_TCHAR(FunctionDeclarationsStack.Last()->getParent()->getNameAsString().c_str()), *FString(FunctionDeclarationsStack.Last()->getNameAsString().c_str()));
+		}
+
+		return true;
+	}
+
+}
diff --git a/Engine/Source/Programs/UnrealCodeAnalyzer/Private/ThreadSafety/TSVisitor.h b/Engine/Source/Programs/UnrealCodeAnalyzer/Private/ThreadSafety/TSVisitor.h
new file mode 100644
--- /dev/null
+++ b/Engine/Source/Programs/UnrealCodeAnalyzer/Private/ThreadSafety/TSVisitor.h
@@ -0,0 +1,25 @@
+// Copyright 1998-2015 Epic Games, Inc. All Rights Reserved.
+
+#pragma once
+#include "clang/AST/RecursiveASTVisitor.h"
+
+namespace UnrealCodeAnalyzer
+{
+	class FTSVisitor : public clang::RecursiveASTVisitor<FTSVisitor>
+	{
+	public:
+		explicit FTSVisitor(clang::ASTContext* InContext, clang::StringRef InFile);
+		bool VisitCXXConstructorDecl(clang::CXXConstructorDecl* Decl);
+		bool TraverseCXXMethodDecl(clang::CXXMethodDecl* Decl);
+		bool VisitDeclRefExpr(clang::DeclRefExpr* Expr);
+		bool VisitCallExpr(clang::CallExpr* Expr);
+
+	private:
+		bool InheritsFromUObject(clang::CXXConstructorDecl* Decl);
+		bool IsGlobalVariable(clang::VarDecl* VarDecl);
+
+		clang::SourceManager& SourceManager;
+		TArray<clang::CXXMethodDecl*> FunctionDeclarationsStack;
+		TSet<clang::Decl*> VisitedDecls;
+	};
+}
diff --git a/Engine/Source/Programs/UnrealCodeAnalyzer/Private/main.cpp b/Engine/Source/Programs/UnrealCodeAnalyzer/Private/main.cpp
--- a/Engine/Source/Programs/UnrealCodeAnalyzer/Private/main.cpp
+++ b/Engine/Source/Programs/UnrealCodeAnalyzer/Private/main.cpp
@@ -4,6 +4,7 @@
 #include "Action.h"
 
 #include "RequiredProgramMainCPPInclude.h"
+#include "ThreadSafety/TSAction.h"
 
 DEFINE_LOG_CATEGORY(LogUnrealCodeAnalyzer)
 
@@ -348,6 +349,24 @@ int AnalyzePCHFile(const FString& CmdLine)
 	return 0;
 }
 
+int32 CheckThreadSafety(int32 ArgC, const char** ArgV)
+{
+	LLVMInitializeX86TargetInfo();
+	LLVMInitializeX86TargetMC();
+	LLVMInitializeX86AsmParser();
+
+	// First argument decides which control flow path we follow, second is output file. They're not relevant from clang's point of view.
+	ArgC -= 2;
+	ArgV[3] = ArgV[2];
+	ArgV[2] = ArgV[0];
+	CommonOptionsParser OptionsParser(ArgC, &(ArgV[2]), MyToolCategory);
+
+	ClangTool Tool(OptionsParser.getCompilations(), OptionsParser.getSourcePathList());
+	Tool.clearArgumentsAdjusters();
+
+	return Tool.run(newFrontendActionFactory<UnrealCodeAnalyzer::FTSAction>().get());
+}
+
 //
 // If analysis/clang compilation gives an error, we still can provide some meaningful data to user,
 // which is why we return 0 in such cases. We return error only when -OutputFile is not provided and
@@ -378,6 +397,11 @@ int main(int32 ArgC, const char* ArgV[])
 		CreateIncludeFiles(ArgC, ArgV);
 	}
 
+	if (FParse::Param(*CmdLine, TEXT("CheckThreadSafety")))
+	{
+		CheckThreadSafety(ArgC, ArgV);
+	}
+
 	if (!FFileHelper::SaveStringToFile(OutputFileContents, *OutputFileName))
 	{
 		return 2;
diff --git a/Engine/Source/Programs/UnrealCodeAnalyzer/UnrealCodeAnalyzer.Build.cs b/Engine/Source/Programs/UnrealCodeAnalyzer/UnrealCodeAnalyzer.Build.cs
--- a/Engine/Source/Programs/UnrealCodeAnalyzer/UnrealCodeAnalyzer.Build.cs
+++ b/Engine/Source/Programs/UnrealCodeAnalyzer/UnrealCodeAnalyzer.Build.cs
@@ -23,6 +23,8 @@ public class UnrealCodeAnalyzer : ModuleRules
 		PrivateIncludePaths.AddRange(new string[] {
 			"Runtime/Launch/Private",	// For LaunchEngineLoop.cpp include
 			"Runtime/Launch/Public",
+			"Programs/UnrealCodeAnalyzer/Public",
+			"Programs/UnrealCodeAnalyzer/Private",
 			"ThirdParty/llvm/3.5.0/include"
 		});
 
